Restricted anim BP compiler subsystem APIs

diff --git a/Engine/Source/Runtime/PropertyAccess/Private/PropertyAccess.cpp b/Engine/Source/Runtime/PropertyAccess/Private/PropertyAccess.cpp
new file mode 100644
--- /dev/null
+++ b/Engine/Source/Runtime/PropertyAccess/Private/PropertyAccess.cpp
@@ -0,0 +1,856 @@
+// Copyright Epic Games, Inc. All Rights Reserved.
+
+#include "PropertyAccess.h"
+#include "Misc/MemStack.h"
+
+#define LOCTEXT_NAMESPACE "PropertyAccess"
+
+struct FPropertyAccessSystem
+{
+	struct FResolveIndirectionsOnLoadContext
+	{
+		FResolveIndirectionsOnLoadContext(TArrayView<FPropertyAccessSegment> InSegments, FPropertyAccessIndirectionChain& InAccess, FPropertyAccessLibrary& InLibrary)
+			: Segments(InSegments)
+			, Access(InAccess)
+			, Library(InLibrary)
+			, AccumulatedOffset(0)
+		{}
+
+		TArrayView<FPropertyAccessSegment> Segments;
+		TArray<FPropertyAccessIndirection> Indirections;
+		FPropertyAccessIndirectionChain& Access;
+		FPropertyAccessLibrary& Library;
+		uint32 AccumulatedOffset;
+		FText ErrorMessage;
+	};
+
+	static uint32 GetPropertyOffset(const FProperty* InProperty, int32 InArrayIndex = 0)
+	{
+		return (uint32)(InProperty->GetOffset_ForInternal() + InProperty->ElementSize * InArrayIndex);
+	}
+
+	// Called on load to resolve all path segments to indirections
+	static bool ResolveIndirectionsOnLoad(FResolveIndirectionsOnLoadContext& InContext)
+	{
+		for(int32 SegmentIndex = 0; SegmentIndex < InContext.Segments.Num(); ++SegmentIndex)
+		{
+			const FPropertyAccessSegment& Segment = InContext.Segments[SegmentIndex];
+			const bool bLastSegment = SegmentIndex == InContext.Segments.Num() - 1;
+
+			if(Segment.Property.Get() == nullptr)
+			{
+				return false;
+			}
+
+			if(EnumHasAllFlags((EPropertyAccessSegmentFlags)Segment.Flags, EPropertyAccessSegmentFlags::Function))
+			{
+				if(Segment.Function == nullptr)
+				{
+					return false;
+				}
+
+				FPropertyAccessIndirection& Indirection = InContext.Indirections.AddDefaulted_GetRef();
+
+				Indirection.Type = Segment.Function->HasAnyFunctionFlags(FUNC_Native) ? EPropertyAccessIndirectionType::NativeFunction : EPropertyAccessIndirectionType::ScriptFunction;
+
+				switch((EPropertyAccessSegmentFlags)Segment.Flags & ~EPropertyAccessSegmentFlags::ModifierFlags)
+				{
+				case EPropertyAccessSegmentFlags::Struct:
+				case EPropertyAccessSegmentFlags::Leaf:
+					Indirection.ObjectType = EPropertyAccessObjectType::None;
+					break;
+				case EPropertyAccessSegmentFlags::Object:
+					Indirection.ObjectType = EPropertyAccessObjectType::Object;
+					break;
+				case EPropertyAccessSegmentFlags::WeakObject:
+					Indirection.ObjectType = EPropertyAccessObjectType::WeakObject;
+					break;
+				case EPropertyAccessSegmentFlags::SoftObject:
+					Indirection.ObjectType = EPropertyAccessObjectType::SoftObject;
+					break;
+				default:
+					check(false);
+					break;
+				}
+
+				Indirection.Function = Segment.Function;
+				Indirection.ReturnBufferSize = Segment.Property.Get()->GetSize();
+				Indirection.ReturnBufferAlignment = Segment.Property.Get()->GetMinAlignment();
+			}
+			else
+			{
+				const int32 ArrayIndex = Segment.ArrayIndex == INDEX_NONE ? 0 : Segment.ArrayIndex;
+				const EPropertyAccessSegmentFlags UnmodifiedFlags = (EPropertyAccessSegmentFlags)Segment.Flags & ~EPropertyAccessSegmentFlags::ModifierFlags;
+				switch(UnmodifiedFlags)
+				{
+				case EPropertyAccessSegmentFlags::Struct:
+				case EPropertyAccessSegmentFlags::Leaf:
+				{
+					FPropertyAccessIndirection& Indirection = InContext.Indirections.AddDefaulted_GetRef();
+
+					Indirection.Offset = GetPropertyOffset(Segment.Property.Get(), ArrayIndex);
+					Indirection.Type = EPropertyAccessIndirectionType::Offset;
+					break;
+				}
+				case EPropertyAccessSegmentFlags::Object:
+				{
+					FPropertyAccessIndirection& Indirection = InContext.Indirections.AddDefaulted_GetRef();
+
+					Indirection.Offset = GetPropertyOffset(Segment.Property.Get(), ArrayIndex);
+					Indirection.Type = EPropertyAccessIndirectionType::Object;
+					Indirection.ObjectType = EPropertyAccessObjectType::Object;
+					break;
+				}
+				case EPropertyAccessSegmentFlags::WeakObject:
+				{
+					FPropertyAccessIndirection& Indirection = InContext.Indirections.AddDefaulted_GetRef();
+
+					Indirection.Offset = GetPropertyOffset(Segment.Property.Get(), ArrayIndex);
+					Indirection.Type = EPropertyAccessIndirectionType::Object;
+					Indirection.ObjectType = EPropertyAccessObjectType::WeakObject;
+					break;
+				}
+				case EPropertyAccessSegmentFlags::SoftObject:
+				{
+					FPropertyAccessIndirection& Indirection = InContext.Indirections.AddDefaulted_GetRef();
+
+					Indirection.Offset = GetPropertyOffset(Segment.Property.Get(), ArrayIndex);
+					Indirection.Type = EPropertyAccessIndirectionType::Object;
+					Indirection.ObjectType = EPropertyAccessObjectType::SoftObject;
+					break;
+				}
+				case EPropertyAccessSegmentFlags::Array:
+				case EPropertyAccessSegmentFlags::ArrayOfStructs:
+				case EPropertyAccessSegmentFlags::ArrayOfObjects:
+				{
+					FPropertyAccessIndirection& Indirection = InContext.Indirections.AddDefaulted_GetRef();
+
+					Indirection.Offset = GetPropertyOffset(Segment.Property.Get());
+					Indirection.Type = EPropertyAccessIndirectionType::Array;
+					Indirection.ArrayProperty = CastFieldChecked<FArrayProperty>(Segment.Property.Get());
+					Indirection.ArrayIndex = ArrayIndex;
+					if(UnmodifiedFlags == EPropertyAccessSegmentFlags::ArrayOfObjects)
+					{
+						if(!bLastSegment)
+						{
+							// Object arrays need an object dereference adding if non-leaf
+							FPropertyAccessIndirection& ExtraIndirection = InContext.Indirections.AddDefaulted_GetRef();
+
+							ExtraIndirection.Offset = 0;
+							ExtraIndirection.Type = EPropertyAccessIndirectionType::Object;
+
+							FProperty* InnerProperty = Indirection.ArrayProperty.Get()->Inner;
+							if(FObjectProperty* ObjectProperty = CastField<FObjectProperty>(InnerProperty))
+							{
+								ExtraIndirection.ObjectType = EPropertyAccessObjectType::Object;
+							}
+							else if(FWeakObjectProperty* WeakObjectProperty = CastField<FWeakObjectProperty>(InnerProperty))
+							{
+								ExtraIndirection.ObjectType = EPropertyAccessObjectType::WeakObject;
+							}
+							else if(FSoftObjectProperty* SoftObjectProperty = CastField<FSoftObjectProperty>(InnerProperty))
+							{
+								ExtraIndirection.ObjectType = EPropertyAccessObjectType::SoftObject;
+							}
+						}
+					}
+					break;
+				}
+				default:
+					check(false);
+					break;
+				}
+			}
+		}
+
+		// Collapse adjacent offset indirections
+		for(int32 IndirectionIndex = 0; IndirectionIndex < InContext.Indirections.Num(); ++IndirectionIndex)
+		{
+			FPropertyAccessIndirection& StartIndirection = InContext.Indirections[IndirectionIndex];
+			if(StartIndirection.Type == EPropertyAccessIndirectionType::Offset)
+			{
+				for(int32 NextIndirectionIndex = IndirectionIndex + 1; NextIndirectionIndex < InContext.Indirections.Num(); ++NextIndirectionIndex)
+				{
+					FPropertyAccessIndirection& RunIndirection = InContext.Indirections[NextIndirectionIndex];
+					if(RunIndirection.Type == EPropertyAccessIndirectionType::Offset)
+					{
+						StartIndirection.Offset += RunIndirection.Offset;
+						InContext.Indirections.RemoveAt(NextIndirectionIndex);
+					}
+					else
+					{
+						// No run, exit
+						break;
+					}
+				}
+			}
+		}
+
+		// Concatenate indirections into the library and update the access
+		InContext.Access.IndirectionStartIndex = InContext.Library.Indirections.Num();
+		InContext.Library.Indirections.Append(InContext.Indirections);
+		InContext.Access.IndirectionEndIndex = InContext.Library.Indirections.Num();
+
+		// Copy leaf property to access
+		FPropertyAccessSegment& LastSegment = InContext.Segments.Last();
+		switch((EPropertyAccessSegmentFlags)LastSegment.Flags & ~EPropertyAccessSegmentFlags::ModifierFlags)
+		{
+		case EPropertyAccessSegmentFlags::Struct:
+		case EPropertyAccessSegmentFlags::Leaf:
+		case EPropertyAccessSegmentFlags::Object:
+		case EPropertyAccessSegmentFlags::WeakObject:
+		case EPropertyAccessSegmentFlags::SoftObject:
+			InContext.Access.Property = LastSegment.Property.Get();
+			break;
+		case EPropertyAccessSegmentFlags::Array:
+		case EPropertyAccessSegmentFlags::ArrayOfStructs:
+		case EPropertyAccessSegmentFlags::ArrayOfObjects:
+			InContext.Access.Property = CastFieldChecked<FArrayProperty>(LastSegment.Property.Get())->Inner;
+			break;
+		}
+
+		return true;
+	}
+
+	static void PostLoadLibrary(FPropertyAccessLibrary& InLibrary)
+	{
+		InLibrary.Indirections.Reset();
+
+		const int32 SrcCount = InLibrary.SrcPaths.Num();
+		InLibrary.SrcAccesses.Reset();
+		InLibrary.SrcAccesses.SetNum(SrcCount);
+		const int32 DestCount = InLibrary.DestPaths.Num();
+		InLibrary.DestAccesses.Reset();
+		InLibrary.DestAccesses.SetNum(DestCount);
+
+		// @TODO: ParallelFor this if required
+
+		for(int32 SrcIndex = 0; SrcIndex < SrcCount; ++SrcIndex)
+		{
+			TArrayView<FPropertyAccessSegment> Segments(&InLibrary.PathSegments[InLibrary.SrcPaths[SrcIndex].PathSegmentStartIndex], InLibrary.SrcPaths[SrcIndex].PathSegmentCount);
+			FResolveIndirectionsOnLoadContext Context(Segments, InLibrary.SrcAccesses[SrcIndex], InLibrary);
+			if(!ResolveIndirectionsOnLoad(Context))
+			{
+				Context.Indirections.Empty();
+				Context.Access.IndirectionStartIndex = Context.Access.IndirectionEndIndex = INDEX_NONE;
+			}
+		}
+
+		for(int32 DestIndex = 0; DestIndex < DestCount; ++DestIndex)
+		{
+			TArrayView<FPropertyAccessSegment> Segments(&InLibrary.PathSegments[InLibrary.DestPaths[DestIndex].PathSegmentStartIndex], InLibrary.DestPaths[DestIndex].PathSegmentCount);
+			FResolveIndirectionsOnLoadContext Context(Segments, InLibrary.DestAccesses[DestIndex], InLibrary);
+			if(!ResolveIndirectionsOnLoad(Context))
+			{
+				Context.Indirections.Empty();
+				Context.Access.IndirectionStartIndex = Context.Access.IndirectionEndIndex = INDEX_NONE;
+			}
+		}
+
+		InLibrary.bHasBeenPostLoaded = true;
+	}
+
+	static void PerformCopy(const FPropertyAccessCopy& Copy, const FProperty* InSrcProperty, const void* InSrcAddr, const FProperty* InDestProperty, void* InDestAddr, TFunctionRef<void(const FProperty*, void*)> InPostCopyOperation)
+	{
+		switch(Copy.Type)
+		{
+		case EPropertyAccessCopyType::Plain:
+			checkSlow(InSrcProperty->PropertyFlags & CPF_IsPlainOldData);
+			checkSlow(InDestProperty->PropertyFlags & CPF_IsPlainOldData);
+			FMemory::Memcpy(InDestAddr, InSrcAddr, InSrcProperty->ElementSize);
+			break;
+		case EPropertyAccessCopyType::Complex:
+			InSrcProperty->CopyCompleteValue(InDestAddr, InSrcAddr);
+			break;
+		case EPropertyAccessCopyType::Bool:
+			checkSlow(InSrcProperty->IsA<FBoolProperty>());
+			checkSlow(InDestProperty->IsA<FBoolProperty>());
+			static_cast<const FBoolProperty*>(InDestProperty)->SetPropertyValue(InDestAddr, static_cast<const FBoolProperty*>(InSrcProperty)->GetPropertyValue(InSrcAddr));
+			break;
+		case EPropertyAccessCopyType::Struct:
+			checkSlow(InSrcProperty->IsA<FStructProperty>());
+			checkSlow(InDestProperty->IsA<FStructProperty>());
+			static_cast<const FStructProperty*>(InDestProperty)->Struct->CopyScriptStruct(InDestAddr, InSrcAddr);
+			break;
+		case EPropertyAccessCopyType::Object:
+			checkSlow(InSrcProperty->IsA<FObjectPropertyBase>());
+			checkSlow(InDestProperty->IsA<FObjectPropertyBase>());
+			static_cast<const FObjectPropertyBase*>(InDestProperty)->SetObjectPropertyValue(InDestAddr, static_cast<const FObjectPropertyBase*>(InSrcProperty)->GetObjectPropertyValue(InSrcAddr));
+			break;
+		case EPropertyAccessCopyType::Name:
+			checkSlow(InSrcProperty->IsA<FNameProperty>());
+			checkSlow(InDestProperty->IsA<FNameProperty>());
+			static_cast<const FNameProperty*>(InDestProperty)->SetPropertyValue(InDestAddr, static_cast<const FNameProperty*>(InSrcProperty)->GetPropertyValue(InSrcAddr));
+			break;
+		case EPropertyAccessCopyType::Array:
+		{
+			checkSlow(InSrcProperty->IsA<FArrayProperty>());
+			checkSlow(InDestProperty->IsA<FArrayProperty>());
+			const FArrayProperty* SrcArrayProperty = static_cast<const FArrayProperty*>(InSrcProperty);
+			const FArrayProperty* DestArrayProperty = static_cast<const FArrayProperty*>(InDestProperty);
+			FScriptArrayHelper SourceArrayHelper(SrcArrayProperty, InSrcAddr);
+			FScriptArrayHelper DestArrayHelper(DestArrayProperty, InDestAddr);
+
+			// Copy the minimum number of elements to the destination array without resizing
+			const int32 MinSize = FMath::Min(SourceArrayHelper.Num(), DestArrayHelper.Num());
+			for(int32 ElementIndex = 0; ElementIndex < MinSize; ++ElementIndex)
+			{
+				SrcArrayProperty->Inner->CopySingleValue(DestArrayHelper.GetRawPtr(ElementIndex), SourceArrayHelper.GetRawPtr(ElementIndex));
+			}
+			break;
+		}
+		case EPropertyAccessCopyType::PromoteBoolToByte:
+			checkSlow(InSrcProperty->IsA<FBoolProperty>());
+			checkSlow(InDestProperty->IsA<FByteProperty>());
+			static_cast<const FByteProperty*>(InDestProperty)->SetPropertyValue(InDestAddr, (uint8)static_cast<const FBoolProperty*>(InSrcProperty)->GetPropertyValue(InSrcAddr));
+			break;
+		case EPropertyAccessCopyType::PromoteBoolToInt32:
+			checkSlow(InSrcProperty->IsA<FBoolProperty>());
+			checkSlow(InDestProperty->IsA<FIntProperty>());
+			static_cast<const FIntProperty*>(InDestProperty)->SetPropertyValue(InDestAddr, (int32)static_cast<const FBoolProperty*>(InSrcProperty)->GetPropertyValue(InSrcAddr));
+			break;
+		case EPropertyAccessCopyType::PromoteBoolToInt64:
+			checkSlow(InSrcProperty->IsA<FBoolProperty>());
+			checkSlow(InDestProperty->IsA<FInt64Property>());
+			static_cast<const FInt64Property*>(InDestProperty)->SetPropertyValue(InDestAddr, (int64)static_cast<const FBoolProperty*>(InSrcProperty)->GetPropertyValue(InSrcAddr));
+			break;
+		case EPropertyAccessCopyType::PromoteBoolToFloat:
+			checkSlow(InSrcProperty->IsA<FBoolProperty>());
+			checkSlow(InDestProperty->IsA<FFloatProperty>());
+			static_cast<const FFloatProperty*>(InDestProperty)->SetPropertyValue(InDestAddr, (float)static_cast<const FBoolProperty*>(InSrcProperty)->GetPropertyValue(InSrcAddr));
+			break;
+		case EPropertyAccessCopyType::PromoteByteToInt32:
+			checkSlow(InSrcProperty->IsA<FByteProperty>());
+			checkSlow(InDestProperty->IsA<FIntProperty>());
+			static_cast<const FIntProperty*>(InDestProperty)->SetPropertyValue(InDestAddr, (int32)static_cast<const FByteProperty*>(InSrcProperty)->GetPropertyValue(InSrcAddr));
+			break;
+		case EPropertyAccessCopyType::PromoteByteToInt64:
+			checkSlow(InSrcProperty->IsA<FByteProperty>());
+			checkSlow(InDestProperty->IsA<FInt64Property>());
+			static_cast<const FInt64Property*>(InDestProperty)->SetPropertyValue(InDestAddr, (int64)static_cast<const FByteProperty*>(InSrcProperty)->GetPropertyValue(InSrcAddr));
+			break;
+		case EPropertyAccessCopyType::PromoteByteToFloat:
+			checkSlow(InSrcProperty->IsA<FByteProperty>());
+			checkSlow(InDestProperty->IsA<FFloatProperty>());
+			static_cast<const FFloatProperty*>(InDestProperty)->SetPropertyValue(InDestAddr, (float)static_cast<const FByteProperty*>(InSrcProperty)->GetPropertyValue(InSrcAddr));
+			break;
+		case EPropertyAccessCopyType::PromoteInt32ToInt64:
+			checkSlow(InSrcProperty->IsA<FIntProperty>());
+			checkSlow(InDestProperty->IsA<FInt64Property>());
+			static_cast<const FInt64Property*>(InDestProperty)->SetPropertyValue(InDestAddr, (int64)static_cast<const FIntProperty*>(InSrcProperty)->GetPropertyValue(InSrcAddr));
+			break;
+		case EPropertyAccessCopyType::PromoteInt32ToFloat:
+			checkSlow(InSrcProperty->IsA<FIntProperty>());
+			checkSlow(InDestProperty->IsA<FFloatProperty>());
+			static_cast<const FFloatProperty*>(InDestProperty)->SetPropertyValue(InDestAddr, (float)static_cast<const FIntProperty*>(InSrcProperty)->GetPropertyValue(InSrcAddr));
+			break;
+		default:
+			check(false);
+			break;
+		}
+
+		InPostCopyOperation(InDestProperty, InDestAddr);
+	}
+
+	static void CallNativeAccessor(UObject* InObject, UFunction* InFunction, void* OutRetValue)
+	{
+		// Func must be local
+		check((InObject->GetFunctionCallspace(InFunction, nullptr) & FunctionCallspace::Local) != 0);		
+
+		// Function must be native
+		check(InFunction->HasAnyFunctionFlags(FUNC_Native));
+
+		// Function must have a return property
+		check(InFunction->GetReturnProperty() != nullptr);
+
+		// Function must only have one param - its return value
+		check(InFunction->NumParms == 1);
+
+		FFrame Stack(InObject, InFunction, nullptr, nullptr, InFunction->ChildProperties);
+		InFunction->Invoke(InObject, Stack, OutRetValue);
+	}
+
+	static void IterateAccess(void* InContainer, const FPropertyAccessLibrary& InLibrary, const FPropertyAccessIndirectionChain& InAccess, TFunctionRef<void(void*)> InAddressFunction, TFunctionRef<void(UObject*)> InPerObjectFunction)
+	{
+		void* Address = InContainer;
+
+		for(int32 IndirectionIndex = InAccess.IndirectionStartIndex; Address != nullptr && IndirectionIndex < InAccess.IndirectionEndIndex; ++IndirectionIndex)
+		{
+			const FPropertyAccessIndirection& Indirection = InLibrary.Indirections[IndirectionIndex];
+
+			switch(Indirection.Type)
+			{
+			case EPropertyAccessIndirectionType::Offset:
+				Address = static_cast<void*>(static_cast<uint8*>(Address) + Indirection.Offset);
+				break;
+			case EPropertyAccessIndirectionType::Object:
+			{
+				switch(Indirection.ObjectType)
+				{
+				case EPropertyAccessObjectType::Object:
+				{
+					UObject* Object = *reinterpret_cast<UObject**>(static_cast<uint8*>(Address) + Indirection.Offset);
+					Address = static_cast<void*>(Object);
+					if(Object != nullptr)
+					{
+						InPerObjectFunction(Object);
+					}
+					break;
+				}
+				case EPropertyAccessObjectType::WeakObject:
+				{
+					TWeakObjectPtr<UObject>& WeakObjectPtr = *reinterpret_cast<TWeakObjectPtr<UObject>*>(static_cast<uint8*>(Address) + Indirection.Offset);
+					UObject* Object = WeakObjectPtr.Get();
+					Address = static_cast<void*>(Object);
+					if(Object != nullptr)
+					{
+						InPerObjectFunction(Object);
+					}
+					break;
+				}
+				case EPropertyAccessObjectType::SoftObject:
+				{
+					FSoftObjectPtr& SoftObjectPtr = *reinterpret_cast<FSoftObjectPtr*>(static_cast<uint8*>(Address) + Indirection.Offset);
+					UObject* Object = SoftObjectPtr.Get();
+					Address = static_cast<void*>(Object);
+					if(Object != nullptr)
+					{
+						InPerObjectFunction(Object);
+					}
+					break;
+				}
+				default:
+					check(false);
+				}
+				break;
+			}
+			case EPropertyAccessIndirectionType::Array:
+			{
+				if(FArrayProperty* ArrayProperty = Indirection.ArrayProperty.Get())
+				{
+					FScriptArrayHelper Helper(ArrayProperty, static_cast<uint8*>(Address) + Indirection.Offset);
+					if(Helper.IsValidIndex(Indirection.ArrayIndex))
+					{
+						Address = static_cast<void*>(Helper.GetRawPtr(Indirection.ArrayIndex));
+					}
+					else
+					{
+						Address = nullptr;
+					}
+				}
+				else
+				{
+					Address = nullptr;
+				}
+				break;
+			}
+			case EPropertyAccessIndirectionType::ScriptFunction:
+			case EPropertyAccessIndirectionType::NativeFunction:
+			{
+				if(Indirection.Function != nullptr)
+				{
+					UObject* CalleeObject = static_cast<UObject*>(Address);
+
+					// Allocate an aligned buffer for the return value
+					Address = (uint8*)FMemStack::Get().Alloc(Indirection.ReturnBufferSize, Indirection.ReturnBufferAlignment);
+
+					// Init value
+					check(Indirection.Function->GetReturnProperty());
+					Indirection.Function->GetReturnProperty()->InitializeValue(Address);
+
+					if(Indirection.Type == EPropertyAccessIndirectionType::NativeFunction)
+					{
+						CallNativeAccessor(CalleeObject, Indirection.Function, Address);
+					}
+					else
+					{
+						CalleeObject->ProcessEvent(Indirection.Function, Address);
+					}
+
+					// Function access may return an object, so we need to follow that ptr
+					switch(Indirection.ObjectType)
+					{
+					case EPropertyAccessObjectType::Object:
+					{
+						UObject* Object = *static_cast<UObject**>(Address);
+						Address = static_cast<void*>(Object);
+						if(Object != nullptr)
+						{
+							InPerObjectFunction(Object);
+						}
+						break;
+					}
+					case EPropertyAccessObjectType::WeakObject:
+					{
+						TWeakObjectPtr<UObject>& WeakObjectPtr = *static_cast<TWeakObjectPtr<UObject>*>(Address);
+						UObject* Object = WeakObjectPtr.Get();
+						Address = static_cast<void*>(Object);
+						if(Object != nullptr)
+						{
+							InPerObjectFunction(Object);
+						}
+						break;
+					}
+					case EPropertyAccessObjectType::SoftObject:
+					{
+						FSoftObjectPtr& SoftObjectPtr = *reinterpret_cast<FSoftObjectPtr*>(Address);
+						UObject* Object = SoftObjectPtr.Get();
+						Address = static_cast<void*>(Object);
+						if(Object != nullptr)
+						{
+							InPerObjectFunction(Object);
+						}
+						break;
+					}
+					default:
+						break;
+					}
+					break;
+				}
+				else
+				{
+					Address = nullptr;
+				}
+			}
+			default:
+				check(false);
+			}
+		}
+
+		if(Address != nullptr)
+		{
+			InAddressFunction(Address);
+		}
+	}
+
+	// Iterates all the objects in an access path/indirection graph.
+	static void ForEachObjectInAccess(void* InContainer, const FPropertyAccessLibrary& InLibrary, const FPropertyAccessIndirectionChain& InAccess, TFunctionRef<void(UObject*)> InPerObjectFunction)
+	{
+		IterateAccess(InContainer, InLibrary, InAccess, [](void*){}, InPerObjectFunction);
+	}
+
+	// Gets the address that corresponds to a property access.
+	// Forwards the address onto the passed-in function. This callback-style approach is used because in some cases 
+	// (e.g. functions), the address may be memory allocated on the stack.
+	static void GetAccessAddress(void* InContainer, const FPropertyAccessLibrary& InLibrary, const FPropertyAccessIndirectionChain& InAccess, TFunctionRef<void(void*)> InAddressFunction)
+	{
+		IterateAccess(InContainer, InLibrary, InAccess, InAddressFunction, [](UObject*){});
+	}
+
+	// Process a single copy
+	static void ProcessCopy(UStruct* InStruct, void* InContainer, const FPropertyAccessLibrary& InLibrary, int32 InCopyIndex, EPropertyAccessCopyBatch InBatchType, TFunctionRef<void(const FProperty*, void*)> InPostCopyOperation)
+	{
+		if(InLibrary.bHasBeenPostLoaded)
+		{
+			const FPropertyAccessCopy& Copy = InLibrary.CopyBatches[(__underlying_type(EPropertyAccessCopyBatch))InBatchType].Copies[InCopyIndex];
+			const FPropertyAccessIndirectionChain& SrcAccess = InLibrary.SrcAccesses[Copy.AccessIndex];
+			if(SrcAccess.Property.Get())
+			{
+				GetAccessAddress(InContainer, InLibrary, SrcAccess, [InContainer, &InLibrary, &Copy, &SrcAccess, &InPostCopyOperation](void* InSrcAddress)
+				{
+					for(int32 DestAccessIndex = Copy.DestAccessStartIndex; DestAccessIndex < Copy.DestAccessEndIndex; ++DestAccessIndex)
+					{
+						const FPropertyAccessIndirectionChain& DestAccess = InLibrary.DestAccesses[DestAccessIndex];
+						if(DestAccess.Property.Get())
+						{
+							GetAccessAddress(InContainer, InLibrary, DestAccess, [&InSrcAddress, &Copy, &SrcAccess, &DestAccess, &InPostCopyOperation](void* InDestAddress)
+							{
+								PerformCopy(Copy, SrcAccess.Property.Get(), InSrcAddress, DestAccess.Property.Get(), InDestAddress, InPostCopyOperation);
+							});
+						}
+					}
+				});
+			}
+		}
+	}
+
+	static void ProcessCopies(UStruct* InStruct, void* InContainer, const FPropertyAccessLibrary& InLibrary, EPropertyAccessCopyBatch InBatchType)
+	{
+		if(InLibrary.bHasBeenPostLoaded)
+		{
+			FMemMark Mark(FMemStack::Get());
+
+			// Copy all valid properties
+			// Parallelization opportunity: ParallelFor all the property copies we need to make
+			const int32 NumCopies = InLibrary.CopyBatches[(__underlying_type(EPropertyAccessCopyBatch))InBatchType].Copies.Num();
+			for(int32 CopyIndex = 0; CopyIndex < NumCopies; ++CopyIndex)
+			{
+				ProcessCopy(InStruct, InContainer, InLibrary, CopyIndex, InBatchType, [](const FProperty*, void*){});
+			}
+		}
+	}
+
+	/** A node in a class property tree */
+	struct FPropertyNode
+	{
+		FPropertyNode(TArray<FPropertyNode>& InPropertyTree, const FProperty* InProperty, FName InName, int32 InId, int32 InParentId)
+			: Property(InProperty)
+			, Name(InName)
+			, Id(InId)
+			, ParentId(InParentId)
+		{
+			if(InParentId != INDEX_NONE)
+			{
+				if(InPropertyTree[InParentId].FirstChildId == INDEX_NONE)
+				{
+					InPropertyTree[InParentId].FirstChildId = InId;
+					InPropertyTree[InParentId].LastChildId = InId;
+				}
+				InPropertyTree[InParentId].LastChildId++;
+			}
+		}
+
+		FPropertyNode(TArray<FPropertyNode>& InPropertyTree, const FProperty* InProperty, int32 InId, int32 InParentId)
+			: FPropertyNode(InPropertyTree, InProperty, InProperty->GetFName(), InId, InParentId)
+		{
+		}
+
+		FPropertyNode(TArray<FPropertyNode>& InPropertyTree, int32 InId)
+			: FPropertyNode(InPropertyTree, nullptr, NAME_None, InId, INDEX_NONE)
+		{
+		}
+
+		const FProperty* Property = nullptr;
+
+		FName Name = NAME_None;
+
+		int32 Id = INDEX_NONE;
+
+		int32 ParentId = INDEX_NONE;
+
+		int32 FirstChildId = INDEX_NONE;
+
+		int32 LastChildId = INDEX_NONE;
+	};
+
+	/** Map of classes to a tree of property nodes, for faster specialized iteration */
+	static TMap<const UClass*, TUniquePtr<TArray<FPropertyNode>>> ClassPropertyTrees;
+
+	/** Finds the event Id associated with a property node identified by InPath */
+	static int32 FindEventId(const TArray<FPropertyNode>& InPropertyTree, TArrayView<const FName> InPath)
+	{
+		if(InPropertyTree.Num() > 0)
+		{
+			int32 NodeIndex = 0;
+			int32 PathIndex = 0;
+
+			for(int32 ChildIndex = InPropertyTree[NodeIndex].FirstChildId; ChildIndex < InPropertyTree[NodeIndex].LastChildId; ++ChildIndex)
+			{
+				if(InPropertyTree[ChildIndex].Name == InPath[PathIndex])
+				{
+					if(PathIndex == InPath.Num() - 1)
+					{
+						// Found the leaf property of the path
+						check(InPropertyTree[ChildIndex].Id == ChildIndex);
+						return ChildIndex;
+					}
+					else
+					{
+						// Recurse into the next level of the tree/path
+						PathIndex++;
+						NodeIndex = ChildIndex;
+						ChildIndex = InPropertyTree[ChildIndex].FirstChildId;
+					}
+				}
+			}
+		}
+
+		return INDEX_NONE;
+	}
+
+	/** Makes a property tree for the specified struct */
+	static void MakeClassPropertyTree(const UStruct* InStruct, TArray<FPropertyNode>& OutPropertyTree, int32 InParentId = INDEX_NONE)
+	{
+		const int32 StartIndex = OutPropertyTree.Num();
+		int32 Index = OutPropertyTree.Num();
+
+		// Add root
+		if(InParentId == INDEX_NONE)
+		{
+			check(OutPropertyTree.Num() == 0);
+			OutPropertyTree.Emplace(OutPropertyTree, Index++);
+			InParentId = 0;
+		}
+
+		// Make sure that all children are contiguous indices
+		for (TFieldIterator<FProperty> It(InStruct, EFieldIteratorFlags::ExcludeSuper); It; ++It)
+		{
+			const FProperty* Property = *It;
+
+			// Skip editor-only properties, as the runtime cant respond to them and we want editor and non-editor IDs to match
+			if((Property->GetFlags() & CPF_EditorOnly) == 0)
+			{
+				OutPropertyTree.Emplace(OutPropertyTree, Property, Index++, InParentId);
+			}
+		}
+
+		// Reset index to repeat iteration for recursion into sub-structs
+		Index = StartIndex;
+		for (TFieldIterator<FProperty> It(InStruct, EFieldIteratorFlags::ExcludeSuper); It; ++It, ++Index)
+		{
+			const FProperty* Property = *It;
+
+			if((Property->GetFlags() & CPF_EditorOnly) == 0)
+			{
+				if(const FStructProperty* StructProperty = CastField<const FStructProperty>(Property))
+				{
+					MakeClassPropertyTree(StructProperty->Struct, OutPropertyTree, Index);
+				}
+				else if(const FArrayProperty* ArrayProperty = CastField<const FArrayProperty>(Property))
+				{
+					if(const FStructProperty* InnerStructProperty = CastField<const FStructProperty>(ArrayProperty->Inner))
+					{
+						MakeClassPropertyTree(InnerStructProperty->Struct, OutPropertyTree, Index);
+					}
+				}
+			}
+		}
+	}
+
+	static const TArray<FPropertyNode>& GetClassPropertyTree(const UClass* InClass)
+	{
+		TArray<FPropertyNode>* PropertyTree = nullptr;
+
+		if(TUniquePtr<TArray<FPropertyNode>>* ExistingPropertyTreePtr = ClassPropertyTrees.Find(InClass))
+		{
+			PropertyTree = ExistingPropertyTreePtr->Get();
+		}
+		else
+		{
+			PropertyTree = ClassPropertyTrees.Emplace(InClass, MakeUnique<TArray<FPropertyNode>>()).Get();
+			MakeClassPropertyTree(InClass, *PropertyTree);
+		}
+
+		return *PropertyTree;
+	}
+
+	static int32 GetEventId(const UClass* InClass, TArrayView<const FName> InPath)
+	{
+		return FindEventId(GetClassPropertyTree(InClass), InPath);
+	}
+
+	// Mapping from a broadcaster's IDs to subscribers IDs
+	struct FSubscriberMap
+	{
+		TArray<int32> Mapping;
+	};
+
+	// Mapping from a broadcaster to various subscribers, the integer indexes SubscriberMappings
+	struct FBroadcasterMap
+	{
+		TMap<TWeakObjectPtr<const UClass>, int32> SubscriberIndices;
+	};
+
+	// Map for looking up a broadcaster/subscriber
+	static TMap<TWeakObjectPtr<const UClass>, FBroadcasterMap> BroadcasterMaps;
+
+	// Array for stable indices of mappings
+	static TArray<FSubscriberMap> SubscriberMappings;
+
+	/** Creates a mapping between a broadcaster classes IDs and a subscriber classes accesses */
+	static int32 CreateBroadcasterSubscriberMapping(const UClass* InBroadcasterClass, const UClass* InSubscriberClass, const FPropertyAccessLibrary& InSubscriberLibrary)
+	{
+		const TArray<FPropertyNode>& BroadcasterPropertyTree = GetClassPropertyTree(InBroadcasterClass);
+	
+		int32 MappingIndex = SubscriberMappings.Num();
+		FBroadcasterMap& BroadcasterMap = BroadcasterMaps.FindOrAdd(InBroadcasterClass);
+		BroadcasterMap.SubscriberIndices.Add(InSubscriberClass, MappingIndex);
+		FSubscriberMap& Mapping = SubscriberMappings.AddDefaulted_GetRef();
+
+		Mapping.Mapping.SetNum(BroadcasterPropertyTree.Num());
+
+		for(int32& Index : Mapping.Mapping)
+		{
+			Index = INDEX_NONE;
+		}
+
+		for(int32 AccessIndex = 0; AccessIndex < InSubscriberLibrary.SrcAccesses.Num(); ++AccessIndex)
+		{
+			const FPropertyAccessIndirectionChain& Access = InSubscriberLibrary.SrcAccesses[AccessIndex];
+			Mapping.Mapping[Access.EventId] = AccessIndex;
+		}
+
+		return MappingIndex;
+	}
+
+	static void BindEvents(UObject* InSubscriberObject, UClass* InSubscriberClass, const FPropertyAccessLibrary& InSubscriberLibrary)
+	{
+		IPropertyEventSubscriber* SubscriberObject = CastChecked<IPropertyEventSubscriber>(InSubscriberObject);
+
+		TSet<UObject*> Broadcasters;
+
+		for(int32 EventAccessIndex : InSubscriberLibrary.EventAccessIndices)
+		{
+			const FPropertyAccessIndirectionChain& Access = InSubscriberLibrary.SrcAccesses[EventAccessIndex];
+			check(Access.EventId != INDEX_NONE);
+
+			// Find all unique broadcasters
+			ForEachObjectInAccess(InSubscriberObject, InSubscriberLibrary, Access, [&Broadcasters](UObject* InBroadcasterObject)
+			{
+				if(IPropertyEventBroadcaster* Broadcaster = Cast<IPropertyEventBroadcaster>(InBroadcasterObject))
+				{
+					Broadcasters.Add(InBroadcasterObject);
+				}
+			});
+
+			for(UObject* BroadcasterObject : Broadcasters)
+			{
+				// lookup or create a mapping
+				int32 MappingIndex = INDEX_NONE;
+				const UClass* BroadcasterClass = BroadcasterObject->GetClass();
+				
+				if(FBroadcasterMap* BroadcasterMapPtr = BroadcasterMaps.Find(BroadcasterClass))
+				{
+					if(int32* SubscriberMappingIndexPtr = BroadcasterMapPtr->SubscriberIndices.Find(InSubscriberClass))
+					{
+						MappingIndex = *SubscriberMappingIndexPtr;
+					}
+				}
+
+				if(MappingIndex == INDEX_NONE)
+				{
+					MappingIndex = CreateBroadcasterSubscriberMapping(BroadcasterClass, InSubscriberClass, InSubscriberLibrary);
+				}
+
+				IPropertyEventBroadcaster* Broadcaster = CastChecked<IPropertyEventBroadcaster>(BroadcasterObject);
+				Broadcaster->RegisterSubscriber(SubscriberObject, MappingIndex);
+			}
+		}
+	}
+};
+
+TMap<const UClass*, TUniquePtr<TArray<FPropertyAccessSystem::FPropertyNode>>> FPropertyAccessSystem::ClassPropertyTrees;
+TMap<TWeakObjectPtr<const UClass>, FPropertyAccessSystem::FBroadcasterMap> FPropertyAccessSystem::BroadcasterMaps;
+TArray<FPropertyAccessSystem::FSubscriberMap> FPropertyAccessSystem::SubscriberMappings;
+
+namespace PropertyAccess
+{
+	void PostLoadLibrary(FPropertyAccessLibrary& InLibrary)
+	{
+		::FPropertyAccessSystem::PostLoadLibrary(InLibrary);
+	}
+
+	void ProcessCopies(UObject* InObject, const FPropertyAccessLibrary& InLibrary, EPropertyAccessCopyBatch InBatchType)
+	{
+		QUICK_SCOPE_CYCLE_COUNTER(PropertyAccess_ProcessCopies);
+
+		::FPropertyAccessSystem::ProcessCopies(InObject->GetClass(), InObject, InLibrary, InBatchType);
+	}
+
+	void ProcessCopy(UObject* InObject, const FPropertyAccessLibrary& InLibrary, EPropertyAccessCopyBatch InBatchType, int32 InCopyIndex, TFunctionRef<void(const FProperty*, void*)> InPostCopyOperation)
+	{
+		::FPropertyAccessSystem::ProcessCopy(InObject->GetClass(), InObject, InLibrary, InCopyIndex, InBatchType, InPostCopyOperation);
+	}
+
+	void BindEvents(UObject* InObject, const FPropertyAccessLibrary& InLibrary)
+	{
+		::FPropertyAccessSystem::BindEvents(InObject, InObject->GetClass(), InLibrary);
+	}
+
+	int32 GetEventId(const UClass* InClass, TArrayView<const FName> InPath)
+	{
+		return ::FPropertyAccessSystem::GetEventId(InClass, InPath);
+	}
+}
+
+#undef LOCTEXT_NAMESPACE
\ No newline at end of file
diff --git a/Engine/Source/Runtime/PropertyAccess/Private/PropertyAccessModule.cpp b/Engine/Source/Runtime/PropertyAccess/Private/PropertyAccessModule.cpp
new file mode 100644
--- /dev/null
+++ b/Engine/Source/Runtime/PropertyAccess/Private/PropertyAccessModule.cpp
@@ -0,0 +1,6 @@
+// Copyright Epic Games, Inc. All Rights Reserved.
+
+#include "CoreMinimal.h"
+#include "Modules/ModuleManager.h"
+
+IMPLEMENT_MODULE(FDefaultModuleImpl, PropertyAccess)
\ No newline at end of file
diff --git a/Engine/Source/Runtime/PropertyAccess/PropertyAccess.Build.cs b/Engine/Source/Runtime/PropertyAccess/PropertyAccess.Build.cs
new file mode 100644
--- /dev/null
+++ b/Engine/Source/Runtime/PropertyAccess/PropertyAccess.Build.cs
@@ -0,0 +1,16 @@
+// Copyright Epic Games, Inc. All Rights Reserved.
+
+using UnrealBuildTool;
+
+public class PropertyAccess : ModuleRules
+{
+	public PropertyAccess(ReadOnlyTargetRules Target) : base(Target)
+	{
+		PrivateDependencyModuleNames.AddRange(new string[] 
+		{ 
+			"Core",
+			"CoreUObject",
+			"PropertyPath",
+		});
+	}
+}
diff --git a/Engine/Source/Runtime/PropertyAccess/Public/IPropertyAccess.h b/Engine/Source/Runtime/PropertyAccess/Public/IPropertyAccess.h
new file mode 100644
--- /dev/null
+++ b/Engine/Source/Runtime/PropertyAccess/Public/IPropertyAccess.h
@@ -0,0 +1,56 @@
+// Copyright Epic Games, Inc. All Rights Reserved.
+
+#pragma once
+
+#include "CoreMinimal.h"
+#include "UObject/Interface.h"
+#include "IPropertyAccess.generated.h"
+
+// The various types of property copy
+UENUM()
+enum class EPropertyAccessCopyBatch : uint8
+{
+	// A copy of internal->internal data, unbatched
+	InternalUnbatched,
+
+	// A copy of external->internal data, unbatched
+	ExternalUnbatched,
+
+	// A copy of internal->internal data, batched
+	InternalBatched,
+
+	// A copy of external->internal data, batched
+	ExternalBatched,
+
+	Count
+};
+
+UINTERFACE(MinimalAPI)
+class UPropertyAccess : public UInterface
+{
+	GENERATED_BODY()
+};
+
+class IPropertyAccess : public IInterface
+{
+	GENERATED_BODY()
+
+public:
+	/** 
+	 * Process a 'tick' of a property access instance. 
+	 * Note internally allocates via FMemStack and pushes its own FMemMark
+	 */
+	virtual void ProcessCopies(UObject* InObject, EPropertyAccessCopyBatch InBatchType) const = 0;
+
+	/** 
+	 * Process a single copy 
+	 * Note that this can potentially allocate via FMemStack, so inserting FMemMark before a number of these calls is recommended
+	 */
+	virtual void ProcessCopy(UObject* InObject, EPropertyAccessCopyBatch InBatchType, int32 InCopyIndex, TFunctionRef<void(const FProperty*, void*)> InPostCopyOperation) const = 0;
+
+	/** Bind all event-type accesses to their respective objects */
+	virtual void BindEvents(UObject* InObject) const = 0;
+
+	/** Resolve a path to an event Id for the specified class */
+	virtual int32 GetEventId(const UClass* InClass, TArrayView<const FName> InPath) const = 0;
+};
diff --git a/Engine/Source/Runtime/PropertyAccess/Public/PropertyEventInterfaces.h b/Engine/Source/Runtime/PropertyAccess/Public/PropertyEventInterfaces.h
new file mode 100644
--- /dev/null
+++ b/Engine/Source/Runtime/PropertyAccess/Public/PropertyEventInterfaces.h
@@ -0,0 +1,52 @@
+// Copyright Epic Games, Inc. All Rights Reserved.
+
+#pragma once
+
+#include "CoreMinimal.h"
+#include "UObject/Interface.h"
+
+#include "PropertyEventInterfaces.generated.h"
+
+namespace PropertyAccess
+{
+	DECLARE_MULTICAST_DELEGATE_TwoParams(FOnPropertyChanged, UObject* /*InObject*/, int32 /*InBroadcastId*/);
+}
+
+UINTERFACE(MinimalAPI)
+class UPropertyEventBroadcaster : public UInterface
+{
+	GENERATED_BODY()
+};
+
+class IPropertyEventSubscriber;
+
+/** Interface used to broadcast property changed events. Inherit from this and call BROADCAST_PROPERTY_CHANGED */
+class IPropertyEventBroadcaster
+{
+	GENERATED_BODY()
+public:
+
+	/** Broadcast a property changing */
+	virtual void BroadcastPropertyChanged(int32 InBroadcastId) const = 0;
+
+	/** Register a subscriber to listen for property changed events */
+	virtual void RegisterSubscriber(IPropertyEventSubscriber* InSubscriber, int32 InMappingId) = 0;
+
+	/** Bind a delegate to the property changed event */
+	virtual void UnregisterSubscriber(IPropertyEventSubscriber* InSubscriber) = 0;
+};
+
+UINTERFACE(MinimalAPI)
+class UPropertyEventSubscriber : public UInterface
+{
+	GENERATED_BODY()
+};
+
+class IPropertyEventSubscriber
+{
+	GENERATED_BODY()
+public:
+
+	/** Handle a property changing */
+	virtual void OnPropertyChanged(IPropertyEventBroadcaster* InObject, int32 InBroadcastId) const = 0;
+};
