Added a dummy deprecated class to avoid load warnings

diff --git a/Engine/Source/Runtime/PropertyAccess/Private/AnimBlueprintClassSubsystem_PropertyAccess.h b/Engine/Source/Runtime/PropertyAccess/Private/AnimBlueprintClassSubsystem_PropertyAccess.h
new file mode 100644
--- /dev/null
+++ b/Engine/Source/Runtime/PropertyAccess/Private/AnimBlueprintClassSubsystem_PropertyAccess.h
@@ -0,0 +1,13 @@
+// Copyright Epic Games, Inc. All Rights Reserved.
+
+#pragma once
+
+#include "CoreMinimal.h"
+#include "AnimBlueprintClassSubsystem_PropertyAccess.generated.h"
+
+// Dummy class to prevent warnings on load of older subsystem implementation
+UCLASS(MinimalAPI, deprecated)
+class UDEPRECATED_AnimBlueprintClassSubsystem_PropertyAccess : public UObject
+{
+	GENERATED_BODY()
+};
\ No newline at end of file
