Removed deprecation to fix cook warning

#rn Aaron.Cox

diff --git a/Engine/Source/Runtime/PropertyAccess/Private/AnimBlueprintClassSubsystem_PropertyAccess.h b/Engine/Source/Runtime/PropertyAccess/Private/AnimBlueprintClassSubsystem_PropertyAccess.h
--- a/Engine/Source/Runtime/PropertyAccess/Private/AnimBlueprintClassSubsystem_PropertyAccess.h
+++ b/Engine/Source/Runtime/PropertyAccess/Private/AnimBlueprintClassSubsystem_PropertyAccess.h
@@ -6,8 +6,8 @@
 #include "AnimBlueprintClassSubsystem_PropertyAccess.generated.h"
 
 // Dummy class to prevent warnings on load of older subsystem implementation
-UCLASS(MinimalAPI, deprecated)
-class UDEPRECATED_AnimBlueprintClassSubsystem_PropertyAccess : public UObject
+UCLASS(MinimalAPI)
+class UAnimBlueprintClassSubsystem_PropertyAccess : public UObject
 {
 	GENERATED_BODY()
 };
\ No newline at end of file
